- Fixing LWC cis warning in LED wall calibration

diff --git a/Engine/Plugins/Experimental/VirtualProduction/LedWallCalibration/Source/LedWallCalibration/Private/LedWallCalibration.cpp b/Engine/Plugins/Experimental/VirtualProduction/LedWallCalibration/Source/LedWallCalibration/Private/LedWallCalibration.cpp
--- a/Engine/Plugins/Experimental/VirtualProduction/LedWallCalibration/Source/LedWallCalibration/Private/LedWallCalibration.cpp
+++ b/Engine/Plugins/Experimental/VirtualProduction/LedWallCalibration/Source/LedWallCalibration/Private/LedWallCalibration.cpp
@@ -180,7 +180,7 @@ namespace LedWallCalibration
 		FStaticMeshAttributes Attributes(*MeshDescription); // MeshDescription can't be const here
 
 		// Vertex 3d positions, indexable by VertexId
-		const TVertexAttributesRef<const FVector> VertexPositions = MeshDescription->VertexAttributes().GetAttributesRef<FVector>(MeshAttribute::Vertex::Position);
+		const TVertexAttributesRef<const FVector3f> VertexPositions = MeshDescription->VertexAttributes().GetAttributesRef<FVector3f>(MeshAttribute::Vertex::Position);
 
 		// UV 2d positions, indexable by VertexInstaceId
 		const TVertexInstanceAttributesRef<const FVector2D> VertexInstanceUVs = MeshDescription->VertexInstanceAttributes().GetAttributesRef<FVector2D>(MeshAttribute::VertexInstance::TextureCoordinate);
