Update BuildVersionSettings that are out of date

#rnx
#Synced-CL 49960049

diff --git a/Engine/Plugins/Experimental/AVCodecs/AMFCodecs/Source/AMFCodecs/AMFCodecs.Build.cs b/Engine/Plugins/Experimental/AVCodecs/AMFCodecs/Source/AMFCodecs/AMFCodecs.Build.cs
--- a/Engine/Plugins/Experimental/AVCodecs/AMFCodecs/Source/AMFCodecs/AMFCodecs.Build.cs
+++ b/Engine/Plugins/Experimental/AVCodecs/AMFCodecs/Source/AMFCodecs/AMFCodecs.Build.cs
@@ -10,7 +10,7 @@ public class AMFCodecs : ModuleRules
 	{
 		// Without these two compilation fails on VS2017 with D8049: command line is too long to fit in debug record.
 		bLegacyPublicIncludePaths = false;
-		DefaultBuildSettings = BuildSettingsVersion.V2;
+		DefaultBuildSettings = BuildSettingsVersion.Latest;
 
 		// PCHUsage = PCHUsageMode.NoPCHs;
 
diff --git a/Engine/Plugins/Experimental/AVCodecs/AMFCodecs/Source/AMFCodecsRHI/AMFCodecsRHI.Build.cs b/Engine/Plugins/Experimental/AVCodecs/AMFCodecs/Source/AMFCodecsRHI/AMFCodecsRHI.Build.cs
--- a/Engine/Plugins/Experimental/AVCodecs/AMFCodecs/Source/AMFCodecsRHI/AMFCodecsRHI.Build.cs
+++ b/Engine/Plugins/Experimental/AVCodecs/AMFCodecs/Source/AMFCodecsRHI/AMFCodecsRHI.Build.cs
@@ -9,7 +9,7 @@ public class AMFCodecsRHI : ModuleRules
 	public AMFCodecsRHI(ReadOnlyTargetRules Target) : base(Target)
 	{
 		bLegacyPublicIncludePaths = false;
-		DefaultBuildSettings = BuildSettingsVersion.V2;
+		DefaultBuildSettings = BuildSettingsVersion.Latest;
 		
 		PrivateDependencyModuleNames.AddRange(new string[] {
 			"Engine",
diff --git a/Engine/Plugins/Experimental/AVCodecs/AVCodecsCore/Source/AVCodecsCoreRHI/AVCodecsCoreRHI.Build.cs b/Engine/Plugins/Experimental/AVCodecs/AVCodecsCore/Source/AVCodecsCoreRHI/AVCodecsCoreRHI.Build.cs
--- a/Engine/Plugins/Experimental/AVCodecs/AVCodecsCore/Source/AVCodecsCoreRHI/AVCodecsCoreRHI.Build.cs
+++ b/Engine/Plugins/Experimental/AVCodecs/AVCodecsCore/Source/AVCodecsCoreRHI/AVCodecsCoreRHI.Build.cs
@@ -9,7 +9,7 @@ public class AVCodecsCoreRHI : ModuleRules
 	public AVCodecsCoreRHI(ReadOnlyTargetRules Target) : base(Target)
 	{
 		bLegacyPublicIncludePaths = false;
-		DefaultBuildSettings = BuildSettingsVersion.V2;
+		DefaultBuildSettings = BuildSettingsVersion.Latest;
 		
 		PrivateDependencyModuleNames.AddRange(new string[] {
 			"Engine",
diff --git a/Engine/Plugins/Experimental/AVCodecs/LibVpxCodecs/Source/LibVpxCodecs/LibVpxCodecs.Build.cs b/Engine/Plugins/Experimental/AVCodecs/LibVpxCodecs/Source/LibVpxCodecs/LibVpxCodecs.Build.cs
--- a/Engine/Plugins/Experimental/AVCodecs/LibVpxCodecs/Source/LibVpxCodecs/LibVpxCodecs.Build.cs
+++ b/Engine/Plugins/Experimental/AVCodecs/LibVpxCodecs/Source/LibVpxCodecs/LibVpxCodecs.Build.cs
@@ -10,7 +10,7 @@ public class LibVpxCodecs : ModuleRules
     {
         // Without these two compilation fails on VS2017 with D8049: command line is too long to fit in debug record.
         bLegacyPublicIncludePaths = false;
-        DefaultBuildSettings = BuildSettingsVersion.V2;
+        DefaultBuildSettings = BuildSettingsVersion.Latest;
 
         // PCHUsage = PCHUsageMode.NoPCHs;
 
diff --git a/Engine/Plugins/Experimental/AVCodecs/NVCodecs/Source/NVCodecs/NVCodecs.Build.cs b/Engine/Plugins/Experimental/AVCodecs/NVCodecs/Source/NVCodecs/NVCodecs.Build.cs
--- a/Engine/Plugins/Experimental/AVCodecs/NVCodecs/Source/NVCodecs/NVCodecs.Build.cs
+++ b/Engine/Plugins/Experimental/AVCodecs/NVCodecs/Source/NVCodecs/NVCodecs.Build.cs
@@ -10,7 +10,7 @@ public class NVCodecs : ModuleRules
 	{
 		// Without these two compilation fails on VS2017 with D8049: command line is too long to fit in debug record.
 		bLegacyPublicIncludePaths = false;
-		DefaultBuildSettings = BuildSettingsVersion.V2;
+		DefaultBuildSettings = BuildSettingsVersion.Latest;
 
 		PrivateDependencyModuleNames.AddRange(new string[] {
 			"Engine",
diff --git a/Engine/Plugins/Experimental/AVCodecs/NVCodecs/Source/NVCodecsRHI/NVCodecsRHI.Build.cs b/Engine/Plugins/Experimental/AVCodecs/NVCodecs/Source/NVCodecsRHI/NVCodecsRHI.Build.cs
--- a/Engine/Plugins/Experimental/AVCodecs/NVCodecs/Source/NVCodecsRHI/NVCodecsRHI.Build.cs
+++ b/Engine/Plugins/Experimental/AVCodecs/NVCodecs/Source/NVCodecsRHI/NVCodecsRHI.Build.cs
@@ -10,7 +10,7 @@ public class NVCodecsRHI : ModuleRules
 	{
 		// Without these two compilation fails on VS2017 with D8049: command line is too long to fit in debug record.
 		bLegacyPublicIncludePaths = false;
-		DefaultBuildSettings = BuildSettingsVersion.V2;
+		DefaultBuildSettings = BuildSettingsVersion.Latest;
 
 		PrivateDependencyModuleNames.AddRange(new string[] {
 			"Core",
diff --git a/Engine/Plugins/Experimental/AVCodecs/NVCodecs/Source/NVDEC/NVDEC.Build.cs b/Engine/Plugins/Experimental/AVCodecs/NVCodecs/Source/NVDEC/NVDEC.Build.cs
--- a/Engine/Plugins/Experimental/AVCodecs/NVCodecs/Source/NVDEC/NVDEC.Build.cs
+++ b/Engine/Plugins/Experimental/AVCodecs/NVCodecs/Source/NVDEC/NVDEC.Build.cs
@@ -10,7 +10,7 @@ public class NVDEC: ModuleRules
 	{
 		// Without these two compilation fails on VS2017 with D8049: command line is too long to fit in debug record.
 		bLegacyPublicIncludePaths = false;
-		DefaultBuildSettings = BuildSettingsVersion.V2;
+		DefaultBuildSettings = BuildSettingsVersion.Latest;
 
 		PrivateDependencyModuleNames.AddRange(new string[] {
 			"Engine",
diff --git a/Engine/Plugins/Experimental/AVCodecs/NVCodecs/Source/NVENC/NVENC.Build.cs b/Engine/Plugins/Experimental/AVCodecs/NVCodecs/Source/NVENC/NVENC.Build.cs
--- a/Engine/Plugins/Experimental/AVCodecs/NVCodecs/Source/NVENC/NVENC.Build.cs
+++ b/Engine/Plugins/Experimental/AVCodecs/NVCodecs/Source/NVENC/NVENC.Build.cs
@@ -10,7 +10,7 @@ public class NVENC: ModuleRules
 	{
 		// Without these two compilation fails on VS2017 with D8049: command line is too long to fit in debug record.
 		bLegacyPublicIncludePaths = false;
-		DefaultBuildSettings = BuildSettingsVersion.V2;
+		DefaultBuildSettings = BuildSettingsVersion.Latest;
 
 		PrivateDependencyModuleNames.AddRange(new string[] {
 			"Engine",
diff --git a/Engine/Plugins/Experimental/AVCodecs/VTCodecs/Source/VTCodecs/VTCodecs.build.cs b/Engine/Plugins/Experimental/AVCodecs/VTCodecs/Source/VTCodecs/VTCodecs.build.cs
--- a/Engine/Plugins/Experimental/AVCodecs/VTCodecs/Source/VTCodecs/VTCodecs.build.cs
+++ b/Engine/Plugins/Experimental/AVCodecs/VTCodecs/Source/VTCodecs/VTCodecs.build.cs
@@ -10,7 +10,7 @@ public class VTCodecs : ModuleRules
     {
         // Without these two compilation fails on VS2017 with D8049: command line is too long to fit in debug record.
         bLegacyPublicIncludePaths = false;
-        DefaultBuildSettings = BuildSettingsVersion.V2;
+		DefaultBuildSettings = BuildSettingsVersion.Latest;
 
         // PCHUsage = PCHUsageMode.NoPCHs;
 
diff --git a/Engine/Plugins/Experimental/AVCodecs/VTCodecs/Source/VTCodecsRHI/VTCodecsRHI.build.cs b/Engine/Plugins/Experimental/AVCodecs/VTCodecs/Source/VTCodecsRHI/VTCodecsRHI.build.cs
--- a/Engine/Plugins/Experimental/AVCodecs/VTCodecs/Source/VTCodecsRHI/VTCodecsRHI.build.cs
+++ b/Engine/Plugins/Experimental/AVCodecs/VTCodecs/Source/VTCodecsRHI/VTCodecsRHI.build.cs
@@ -9,7 +9,7 @@ public class VTCodecsRHI : ModuleRules
     public VTCodecsRHI(ReadOnlyTargetRules Target) : base(Target)
     {
         bLegacyPublicIncludePaths = false;
-        DefaultBuildSettings = BuildSettingsVersion.V2;
+		DefaultBuildSettings = BuildSettingsVersion.Latest;
         
         PrivateDependencyModuleNames.AddRange(new string[] {
             "Core",
diff --git a/Engine/Plugins/Experimental/AVCodecs/WMFCodecs/Source/WMFCodecs/WMFCodecs.Build.cs b/Engine/Plugins/Experimental/AVCodecs/WMFCodecs/Source/WMFCodecs/WMFCodecs.Build.cs
--- a/Engine/Plugins/Experimental/AVCodecs/WMFCodecs/Source/WMFCodecs/WMFCodecs.Build.cs
+++ b/Engine/Plugins/Experimental/AVCodecs/WMFCodecs/Source/WMFCodecs/WMFCodecs.Build.cs
@@ -9,7 +9,7 @@ public class WMFCodecs : ModuleRules
 	public WMFCodecs(ReadOnlyTargetRules Target) : base(Target)
 	{
 		bLegacyPublicIncludePaths = false;
-		DefaultBuildSettings = BuildSettingsVersion.V2;
+		DefaultBuildSettings = BuildSettingsVersion.Latest;
 
 		PrivateDependencyModuleNames.AddRange(new string[] {
 			"Engine",
diff --git a/Engine/Plugins/Experimental/MutablePopulation/Source/CustomizableObjectPopulation/CustomizableObjectPopulation.Build.cs b/Engine/Plugins/Experimental/MutablePopulation/Source/CustomizableObjectPopulation/CustomizableObjectPopulation.Build.cs
--- a/Engine/Plugins/Experimental/MutablePopulation/Source/CustomizableObjectPopulation/CustomizableObjectPopulation.Build.cs
+++ b/Engine/Plugins/Experimental/MutablePopulation/Source/CustomizableObjectPopulation/CustomizableObjectPopulation.Build.cs
@@ -8,7 +8,7 @@ public class CustomizableObjectPopulation : ModuleRules
 	{
 		ShortName = "MuCOP";
 
-		DefaultBuildSettings = BuildSettingsVersion.V2;
+		DefaultBuildSettings = BuildSettingsVersion.Latest;
 
         PrivateDependencyModuleNames.AddRange(new string[] {
             "InputCore",
diff --git a/Engine/Plugins/Experimental/MutablePopulation/Source/CustomizableObjectPopulationEditor/CustomizableObjectPopulationEditor.Build.cs b/Engine/Plugins/Experimental/MutablePopulation/Source/CustomizableObjectPopulationEditor/CustomizableObjectPopulationEditor.Build.cs
--- a/Engine/Plugins/Experimental/MutablePopulation/Source/CustomizableObjectPopulationEditor/CustomizableObjectPopulationEditor.Build.cs
+++ b/Engine/Plugins/Experimental/MutablePopulation/Source/CustomizableObjectPopulationEditor/CustomizableObjectPopulationEditor.Build.cs
@@ -10,7 +10,7 @@ public class CustomizableObjectPopulationEditor : ModuleRules
 	{
 		ShortName = "MuCOPE";
 
-		DefaultBuildSettings = BuildSettingsVersion.V2;
+		DefaultBuildSettings = BuildSettingsVersion.Latest;
 
 		PrivateDependencyModuleNames.AddRange(
 			new string[] {
diff --git a/Engine/Plugins/Experimental/VirtualProduction/LedWallCalibration/Source/LedWallCalibration/LedWallCalibration.Build.cs b/Engine/Plugins/Experimental/VirtualProduction/LedWallCalibration/Source/LedWallCalibration/LedWallCalibration.Build.cs
--- a/Engine/Plugins/Experimental/VirtualProduction/LedWallCalibration/Source/LedWallCalibration/LedWallCalibration.Build.cs
+++ b/Engine/Plugins/Experimental/VirtualProduction/LedWallCalibration/Source/LedWallCalibration/LedWallCalibration.Build.cs
@@ -6,7 +6,7 @@ public class LedWallCalibration : ModuleRules
 {
 	public LedWallCalibration(ReadOnlyTargetRules Target) : base(Target)
 	{
-		DefaultBuildSettings = BuildSettingsVersion.V2;
+		DefaultBuildSettings = BuildSettingsVersion.Latest;
 
 		PrivateDependencyModuleNames.AddRange(
 			new string[]
diff --git a/Engine/Plugins/Experimental/VirtualProduction/LedWallCalibration/Source/LedWallCalibrationEditor/LedWallCalibrationEditor.Build.cs b/Engine/Plugins/Experimental/VirtualProduction/LedWallCalibration/Source/LedWallCalibrationEditor/LedWallCalibrationEditor.Build.cs
--- a/Engine/Plugins/Experimental/VirtualProduction/LedWallCalibration/Source/LedWallCalibrationEditor/LedWallCalibrationEditor.Build.cs
+++ b/Engine/Plugins/Experimental/VirtualProduction/LedWallCalibration/Source/LedWallCalibrationEditor/LedWallCalibrationEditor.Build.cs
@@ -11,7 +11,7 @@ public class LedWallCalibrationEditor : ModuleRules
 		bUseUnity = false;
 		PCHUsage = PCHUsageMode.NoPCHs;
 
-		DefaultBuildSettings = BuildSettingsVersion.V2;
+		DefaultBuildSettings = BuildSettingsVersion.Latest;
 
 		PrivateDependencyModuleNames.AddRange(
 			new string[]
diff --git a/Engine/Plugins/Media/HardwareEncoders/Source/EncoderAMF/EncoderAMF.Build.cs b/Engine/Plugins/Media/HardwareEncoders/Source/EncoderAMF/EncoderAMF.Build.cs
--- a/Engine/Plugins/Media/HardwareEncoders/Source/EncoderAMF/EncoderAMF.Build.cs
+++ b/Engine/Plugins/Media/HardwareEncoders/Source/EncoderAMF/EncoderAMF.Build.cs
@@ -10,7 +10,7 @@ public class EncoderAMF : ModuleRules
 	{
 		// Without these two compilation fails on VS2017 with D8049: command line is too long to fit in debug record.
 		bLegacyPublicIncludePaths = false;
-		DefaultBuildSettings = BuildSettingsVersion.V2;
+		DefaultBuildSettings = BuildSettingsVersion.Latest;
 
 		// PCHUsage = PCHUsageMode.NoPCHs;
 
diff --git a/Engine/Plugins/Media/HardwareEncoders/Source/EncoderNVENC/EncoderNVENC.Build.cs b/Engine/Plugins/Media/HardwareEncoders/Source/EncoderNVENC/EncoderNVENC.Build.cs
--- a/Engine/Plugins/Media/HardwareEncoders/Source/EncoderNVENC/EncoderNVENC.Build.cs
+++ b/Engine/Plugins/Media/HardwareEncoders/Source/EncoderNVENC/EncoderNVENC.Build.cs
@@ -10,7 +10,7 @@ public class EncoderNVENC: ModuleRules
 	{
 		// Without these two compilation fails on VS2017 with D8049: command line is too long to fit in debug record.
 		bLegacyPublicIncludePaths = false;
-		DefaultBuildSettings = BuildSettingsVersion.V2;
+		DefaultBuildSettings = BuildSettingsVersion.Latest;
 
 		PrivateDependencyModuleNames.AddRange(new string[] {
 			"Engine",
diff --git a/Engine/Plugins/Mutable/Source/CustomizableObject/CustomizableObject.Build.cs b/Engine/Plugins/Mutable/Source/CustomizableObject/CustomizableObject.Build.cs
--- a/Engine/Plugins/Mutable/Source/CustomizableObject/CustomizableObject.Build.cs
+++ b/Engine/Plugins/Mutable/Source/CustomizableObject/CustomizableObject.Build.cs
@@ -8,7 +8,7 @@ public class CustomizableObject : ModuleRules
 	{
 		ShortName = "MuCO";
 
-		DefaultBuildSettings = BuildSettingsVersion.V2;
+		DefaultBuildSettings = BuildSettingsVersion.Latest;
 
 		//bUseUnity = false;
 
diff --git a/Engine/Plugins/Mutable/Source/CustomizableObjectEditor/CustomizableObjectEditor.Build.cs b/Engine/Plugins/Mutable/Source/CustomizableObjectEditor/CustomizableObjectEditor.Build.cs
--- a/Engine/Plugins/Mutable/Source/CustomizableObjectEditor/CustomizableObjectEditor.Build.cs
+++ b/Engine/Plugins/Mutable/Source/CustomizableObjectEditor/CustomizableObjectEditor.Build.cs
@@ -10,7 +10,7 @@ public class CustomizableObjectEditor : ModuleRules
 	{
 		ShortName = "MuCOE";
 
-		DefaultBuildSettings = BuildSettingsVersion.V2;
+		DefaultBuildSettings = BuildSettingsVersion.Latest;
 		//bUseUnity = false;
 
 		PrivateIncludePathModuleNames.AddRange(
diff --git a/Engine/Plugins/Mutable/Source/MutableRuntime/MutableRuntime.Build.cs b/Engine/Plugins/Mutable/Source/MutableRuntime/MutableRuntime.Build.cs
--- a/Engine/Plugins/Mutable/Source/MutableRuntime/MutableRuntime.Build.cs
+++ b/Engine/Plugins/Mutable/Source/MutableRuntime/MutableRuntime.Build.cs
@@ -10,7 +10,7 @@ namespace UnrealBuildTool.Rules
 		{
 			ShortName = "MuR";
 
-			DefaultBuildSettings = BuildSettingsVersion.V2;
+			DefaultBuildSettings = BuildSettingsVersion.Latest;
 			IWYUSupport = IWYUSupport.KeepAsIsForNow;
 			//bUseUnity = false;
 
diff --git a/Engine/Plugins/Mutable/Source/MutableTools/MutableTools.Build.cs b/Engine/Plugins/Mutable/Source/MutableTools/MutableTools.Build.cs
--- a/Engine/Plugins/Mutable/Source/MutableTools/MutableTools.Build.cs
+++ b/Engine/Plugins/Mutable/Source/MutableTools/MutableTools.Build.cs
@@ -10,7 +10,7 @@ namespace UnrealBuildTool.Rules
 		{
 			ShortName = "MuT";
 
-			DefaultBuildSettings = BuildSettingsVersion.V2;
+			DefaultBuildSettings = BuildSettingsVersion.Latest;
 			IWYUSupport = IWYUSupport.KeepAsIsForNow;
 
 			PublicDependencyModuleNames.AddRange(
diff --git a/Engine/Plugins/Mutable/Source/MutableValidation/MutableValidation.Build.cs b/Engine/Plugins/Mutable/Source/MutableValidation/MutableValidation.Build.cs
--- a/Engine/Plugins/Mutable/Source/MutableValidation/MutableValidation.Build.cs
+++ b/Engine/Plugins/Mutable/Source/MutableValidation/MutableValidation.Build.cs
@@ -11,7 +11,7 @@ namespace UnrealBuildTool.Rules
 		{
 			ShortName = "MuV";
 
-			DefaultBuildSettings = BuildSettingsVersion.V2;
+			DefaultBuildSettings = BuildSettingsVersion.Latest;
 
 			PublicDependencyModuleNames.AddRange(
 				new string[] {
diff --git a/Engine/Plugins/VirtualProduction/VirtualCamera/Source/VirtualCamera/VirtualCamera.Build.cs b/Engine/Plugins/VirtualProduction/VirtualCamera/Source/VirtualCamera/VirtualCamera.Build.cs
--- a/Engine/Plugins/VirtualProduction/VirtualCamera/Source/VirtualCamera/VirtualCamera.Build.cs
+++ b/Engine/Plugins/VirtualProduction/VirtualCamera/Source/VirtualCamera/VirtualCamera.Build.cs
@@ -6,7 +6,7 @@ public class VirtualCamera : ModuleRules
 {
 	public VirtualCamera(ReadOnlyTargetRules Target) : base(Target)
 	{
-		DefaultBuildSettings = BuildSettingsVersion.V2;
+		DefaultBuildSettings = BuildSettingsVersion.Latest;
 
 		PublicDependencyModuleNames.AddRange(
 			new string[]
diff --git a/Engine/Plugins/VirtualProduction/VirtualCamera/Source/VirtualCameraEditor/VirtualCameraEditor.Build.cs b/Engine/Plugins/VirtualProduction/VirtualCamera/Source/VirtualCameraEditor/VirtualCameraEditor.Build.cs
--- a/Engine/Plugins/VirtualProduction/VirtualCamera/Source/VirtualCameraEditor/VirtualCameraEditor.Build.cs
+++ b/Engine/Plugins/VirtualProduction/VirtualCamera/Source/VirtualCameraEditor/VirtualCameraEditor.Build.cs
@@ -6,7 +6,7 @@ public class VirtualCameraEditor : ModuleRules
 {
 	public VirtualCameraEditor(ReadOnlyTargetRules Target) : base(Target)
 	{
-		DefaultBuildSettings = BuildSettingsVersion.V2;
+		DefaultBuildSettings = BuildSettingsVersion.Latest;
 
 		PrivateDependencyModuleNames.AddRange(
 			new string[]
diff --git a/Engine/Source/Runtime/AVEncoder/AVEncoder.Build.cs b/Engine/Source/Runtime/AVEncoder/AVEncoder.Build.cs
--- a/Engine/Source/Runtime/AVEncoder/AVEncoder.Build.cs
+++ b/Engine/Source/Runtime/AVEncoder/AVEncoder.Build.cs
@@ -12,7 +12,7 @@ public class AVEncoder : ModuleRules
 	{
 		// Without these two compilation fails on VS2017 with D8049: command line is too long to fit in debug record.
 		bLegacyPublicIncludePaths = false;
-		DefaultBuildSettings = BuildSettingsVersion.V2;
+		DefaultBuildSettings = BuildSettingsVersion.Latest;
 
 		// PCHUsage = PCHUsageMode.NoPCHs;
 
